LandscapeBlueprintBrush: show Cooking category so that the EditorOnly flag can be edited.

diff --git a/Engine/Source/Editor/LandscapeEditorUtilities/Public/LandscapeBlueprintBrush.h b/Engine/Source/Editor/LandscapeEditorUtilities/Public/LandscapeBlueprintBrush.h
--- a/Engine/Source/Editor/LandscapeEditorUtilities/Public/LandscapeBlueprintBrush.h
+++ b/Engine/Source/Editor/LandscapeEditorUtilities/Public/LandscapeBlueprintBrush.h
@@ -9,7 +9,7 @@
 
 #include "LandscapeBlueprintBrush.generated.h"
 
-UCLASS(Abstract, Blueprintable, hidecategories = (Replication, Input, LOD, Actor, Cooking, Rendering))
+UCLASS(Abstract, Blueprintable, hidecategories = (Replication, Input, LOD, Actor, Rendering), showcategories = (Cooking))
 class LANDSCAPEEDITORUTILITIES_API ALandscapeBlueprintBrush : public ALandscapeBlueprintBrushBase
 {
 	GENERATED_UCLASS_BODY()
